Blank C++ template project now correctly shows the Compile button in the level editor toolbar.

diff --git a/Templates/TP_Blank/TP_Blank.uproject b/Templates/TP_Blank/TP_Blank.uproject
--- a/Templates/TP_Blank/TP_Blank.uproject
+++ b/Templates/TP_Blank/TP_Blank.uproject
@@ -3,5 +3,11 @@
 	"EngineAssociation": "",
 	"Category": "",
 	"Description": "",
-	"Modules": []
+	"Modules": [
+		{
+			"Name": "TP_Blank",
+			"Type": "Runtime",
+			"LoadingPhase": "Default"
+		}
+	]
 }
\ No newline at end of file
