Change filetype

diff --git a/Templates/TP_Blank/manifest.json b/Templates/TP_Blank/manifest.json
--- a/Templates/TP_Blank/manifest.json
+++ b/Templates/TP_Blank/manifest.json
@@ -1,4 +1,4 @@
-{
+﻿{
 	"Version": 1,
 	"Name":
 	[
